Fixing crash with microsoft spatial sound

- checking windows version and logging if not the right version for the mixed reality interpop lib
- enabling ms spatial sound plugin for windows

diff --git a/Engine/Plugins/Runtime/MicrosoftSpatialAudio/MicroSoftSpatialSound.uplugin b/Engine/Plugins/Runtime/MicrosoftSpatialAudio/MicroSoftSpatialSound.uplugin
--- a/Engine/Plugins/Runtime/MicrosoftSpatialAudio/MicroSoftSpatialSound.uplugin
+++ b/Engine/Plugins/Runtime/MicrosoftSpatialAudio/MicroSoftSpatialSound.uplugin
@@ -14,7 +14,8 @@
 	"IsBetaVersion": false,
 	"Installed": false,
 	"SupportedTargetPlatforms": [
-		"HoloLens"
+		"HoloLens",
+		"Win64"
 	],
 	"Modules": [
 		{
@@ -22,7 +23,8 @@
 			"Type": "Runtime",
 			"LoadingPhase": "PreDefault",
 			"WhitelistPlatforms": [
-				"HoloLens"
+				"HoloLens",
+				"Win64"
 			]
 		}
 	]
diff --git a/Engine/Plugins/Runtime/MicrosoftSpatialAudio/Source/MicrosoftSpatialSound/Private/MicrosoftSpatialSoundPlugin.cpp b/Engine/Plugins/Runtime/MicrosoftSpatialAudio/Source/MicrosoftSpatialSound/Private/MicrosoftSpatialSoundPlugin.cpp
--- a/Engine/Plugins/Runtime/MicrosoftSpatialAudio/Source/MicrosoftSpatialSound/Private/MicrosoftSpatialSoundPlugin.cpp
+++ b/Engine/Plugins/Runtime/MicrosoftSpatialAudio/Source/MicrosoftSpatialSound/Private/MicrosoftSpatialSoundPlugin.cpp
@@ -2,10 +2,12 @@
 
 #include "MicrosoftSpatialSoundPlugin.h"
 #include "Features/IModularFeatures.h"
+#include "HAL/PlatformMisc.h"
 
 DEFINE_LOG_CATEGORY_STATIC(LogMicrosoftSpatialSound, Verbose, All);
 
-PRAGMA_DISABLE_OPTIMIZATION
+
+#define MIN_WIN_10_VERSION_FOR_WMR_SPATSOUND 1809
 
 static const float UnrealUnitsToMeters = 0.01f;
 
@@ -188,7 +190,6 @@ void FMicrosoftSpatialSound::Initialize(const FAudioPluginInitializationParams I
 	MinFramesRequiredPerObjectUpdate = InitParams.SampleRate / 100;
 
 	SAC = WindowsMixedReality::SpatialAudioClient::CreateSpatialAudioClient();
-
 	if (SAC)
 	{
 		bool bSuccess = SAC->Start(InitParams.NumSources, InitParams.SampleRate);
@@ -385,6 +386,29 @@ uint32 FMicrosoftSpatialSound::Run()
 
 void FMicrosoftSpatialSoundModule::StartupModule()
 {
+#if PLATFORM_WINDOWS
+	FString OSVersionLabel;
+	FString OSSubVersionLabel;
+	FPlatformMisc::GetOSVersions(OSVersionLabel, OSSubVersionLabel);
+
+	// GetOSVersion returns the Win10 release version in the OSVersion rather than the OSSubVersion, so parse it out ourselves
+	OSSubVersionLabel = OSVersionLabel;
+	OSSubVersionLabel.RemoveFromStart("Windows 10 (Release ");
+	OSSubVersionLabel.RemoveFromEnd(")");
+	int32 CurrentVersionNumber = FCString::Atoi(*OSSubVersionLabel);
+
+	if (CurrentVersionNumber < MIN_WIN_10_VERSION_FOR_WMR_SPATSOUND)
+	{
+		UE_LOG(LogMicrosoftSpatialSound,
+			Warning,
+			TEXT("Microsoft Spatial Sound for UE4 currently only supports windows version '%d' or higher (Current version: '%d')"),
+			MIN_WIN_10_VERSION_FOR_WMR_SPATSOUND,
+			CurrentVersionNumber);
+
+			return;
+	}
+#endif
+
 	IModularFeatures::Get().RegisterModularFeature(FMicrosoftSpatialSoundPluginFactory::GetModularFeatureName(), &PluginFactory);
 }
 
@@ -392,6 +416,4 @@ void FMicrosoftSpatialSoundModule::ShutdownModule()
 {
 }
 
-PRAGMA_ENABLE_OPTIMIZATION
-
 IMPLEMENT_MODULE(FMicrosoftSpatialSoundModule, MicrosoftSpatialSound)
\ No newline at end of file
