Fix for windows spatial sound plugin

diff --git a/Engine/Plugins/Runtime/MicrosoftSpatialAudio/Source/MicrosoftSpatialSound/Private/MicrosoftSpatialSoundPlugin.cpp b/Engine/Plugins/Runtime/MicrosoftSpatialAudio/Source/MicrosoftSpatialSound/Private/MicrosoftSpatialSoundPlugin.cpp
--- a/Engine/Plugins/Runtime/MicrosoftSpatialAudio/Source/MicrosoftSpatialSound/Private/MicrosoftSpatialSoundPlugin.cpp
+++ b/Engine/Plugins/Runtime/MicrosoftSpatialAudio/Source/MicrosoftSpatialSound/Private/MicrosoftSpatialSoundPlugin.cpp
@@ -8,6 +8,7 @@ DEFINE_LOG_CATEGORY_STATIC(LogMicrosoftSpatialSound, Verbose, All);
 
 
 #define MIN_WIN_10_VERSION_FOR_WMR_SPATSOUND 1809
+#define WINDOWS_MIXED_REALITY_DEBUG_DLL_SPATSOUND 0
 
 static const float UnrealUnitsToMeters = 0.01f;
 
@@ -407,6 +408,24 @@ void FMicrosoftSpatialSoundModule::StartupModule()
 
 			return;
 	}
+
+	// Load the mixed reality interop library
+	FString EngineDir = FPaths::EngineDir();
+	FString BinariesSubDir = FPlatformProcess::GetBinariesSubdirectory();
+#if WINDOWS_MIXED_REALITY_DEBUG_DLL_SPATSOUND
+	FString DLLName(TEXT("MixedRealityInteropDebug.dll"));
+#else // WINDOWS_MIXED_REALITY_DEBUG_DLL
+	FString DLLName(TEXT("MixedRealityInterop.dll"));
+#endif // WINDOWS_MIXED_REALITY_DEBUG_DLL
+	FString MRInteropLibraryPath = EngineDir / "Binaries/ThirdParty/MixedRealityInteropLibrary" / BinariesSubDir / DLLName;
+	// Then finally try to load the WMR Interop Library
+	void* MixedRealityInteropLibraryHandle = !MRInteropLibraryPath.IsEmpty() ? FPlatformProcess::GetDllHandle(*MRInteropLibraryPath) : nullptr;
+	if (!MixedRealityInteropLibraryHandle)
+	{
+		UE_LOG(LogMicrosoftSpatialSound, Warning, TEXT("Failed to load the microsoft mixed reality interop DLL"));
+
+		return;
+	}
 #endif
 
 	IModularFeatures::Get().RegisterModularFeature(FMicrosoftSpatialSoundPluginFactory::GetModularFeatureName(), &PluginFactory);
