re-disabling microsoft spatial sound plugin for win64

diff --git a/Engine/Plugins/Runtime/MicrosoftSpatialAudio/MicroSoftSpatialSound.uplugin b/Engine/Plugins/Runtime/MicrosoftSpatialAudio/MicroSoftSpatialSound.uplugin
--- a/Engine/Plugins/Runtime/MicrosoftSpatialAudio/MicroSoftSpatialSound.uplugin
+++ b/Engine/Plugins/Runtime/MicrosoftSpatialAudio/MicroSoftSpatialSound.uplugin
@@ -14,8 +14,7 @@
 	"IsBetaVersion": false,
 	"Installed": false,
 	"SupportedTargetPlatforms": [
-		"HoloLens",
-		"Win64"
+		"HoloLens"
 	],
 	"Modules": [
 		{
@@ -23,8 +22,7 @@
 			"Type": "Runtime",
 			"LoadingPhase": "PreDefault",
 			"WhitelistPlatforms": [
-				"HoloLens",
-				"Win64"
+				"HoloLens"
 			]
 		}
 	]
