UE-79491 Using the MicrosoftSpatialSound Plugin as the Spatialization Plugin for Windows on an 1809 Build of Windows Crashes the Editor
-MixedRealityInterop.dll was failing to load because it cannot automatically find all of its dependencies.  Added explicit load of those to MicrosoftSpatialAudio dll loading code.

diff --git a/Engine/Plugins/Runtime/MicrosoftSpatialAudio/MicroSoftSpatialSound.uplugin b/Engine/Plugins/Runtime/MicrosoftSpatialAudio/MicroSoftSpatialSound.uplugin
--- a/Engine/Plugins/Runtime/MicrosoftSpatialAudio/MicroSoftSpatialSound.uplugin
+++ b/Engine/Plugins/Runtime/MicrosoftSpatialAudio/MicroSoftSpatialSound.uplugin
@@ -14,7 +14,8 @@
 	"IsBetaVersion": false,
 	"Installed": false,
 	"SupportedTargetPlatforms": [
-		"HoloLens"
+		"HoloLens",
+		"Win64"
 	],
 	"Modules": [
 		{
@@ -22,7 +23,8 @@
 			"Type": "Runtime",
 			"LoadingPhase": "PreDefault",
 			"WhitelistPlatforms": [
-				"HoloLens"
+				"HoloLens",
+				"Win64"
 			]
 		}
 	]
diff --git a/Engine/Plugins/Runtime/MicrosoftSpatialAudio/Source/MicrosoftSpatialSound/Private/MicrosoftSpatialSoundPlugin.cpp b/Engine/Plugins/Runtime/MicrosoftSpatialAudio/Source/MicrosoftSpatialSound/Private/MicrosoftSpatialSoundPlugin.cpp
--- a/Engine/Plugins/Runtime/MicrosoftSpatialAudio/Source/MicrosoftSpatialSound/Private/MicrosoftSpatialSoundPlugin.cpp
+++ b/Engine/Plugins/Runtime/MicrosoftSpatialAudio/Source/MicrosoftSpatialSound/Private/MicrosoftSpatialSoundPlugin.cpp
@@ -418,6 +418,16 @@ void FMicrosoftSpatialSoundModule::StartupModule()
 	FString DLLName(TEXT("MixedRealityInterop.dll"));
 #endif // WINDOWS_MIXED_REALITY_DEBUG_DLL
 	FString MRInteropLibraryPath = EngineDir / "Binaries/ThirdParty/MixedRealityInteropLibrary" / BinariesSubDir / DLLName;
+
+#if PLATFORM_64BITS
+	// Load these dependencies first or MixedRealityInteropLibraryHandle fails to load since it doesn't look in the correct path for its dependencies automatically
+	FString HoloLensLibraryDir = EngineDir / "Binaries/ThirdParty/Windows/x64";
+	FPlatformProcess::PushDllDirectory(*HoloLensLibraryDir);
+	FPlatformProcess::GetDllHandle(_TEXT("PerceptionDevice.dll"));
+	FPlatformProcess::GetDllHandle(_TEXT("Microsoft.Holographic.AppRemoting.dll"));
+	FPlatformProcess::PopDllDirectory(*HoloLensLibraryDir);
+#endif // PLATFORM_64BITS && WITH_EDITOR	
+	
 	// Then finally try to load the WMR Interop Library
 	void* MixedRealityInteropLibraryHandle = !MRInteropLibraryPath.IsEmpty() ? FPlatformProcess::GetDllHandle(*MRInteropLibraryPath) : nullptr;
 	if (!MixedRealityInteropLibraryHandle)
