Removed Win32 (rather than commenting it out) in MSSpatialAudio.uplugin

diff --git a/Engine/Plugins/Runtime/MicrosoftSpatialAudio/MicroSoftSpatialSound.uplugin b/Engine/Plugins/Runtime/MicrosoftSpatialAudio/MicroSoftSpatialSound.uplugin
--- a/Engine/Plugins/Runtime/MicrosoftSpatialAudio/MicroSoftSpatialSound.uplugin
+++ b/Engine/Plugins/Runtime/MicrosoftSpatialAudio/MicroSoftSpatialSound.uplugin
@@ -14,7 +14,6 @@
 	"IsBetaVersion": false,
 	"Installed": false,
 	"SupportedTargetPlatforms": [
-		//"Win32",
 		"Win64",
 		"XboxOne"
 	],
@@ -25,7 +24,6 @@
 			"LoadingPhase": "PreDefault",
 			"WhitelistPlatforms" :
 			[
-				//"Win32",
 				"Win64",
 				"XboxOne"
 			]
