Removing Win64 as valid platform for MSSpatialSound.uplugin

diff --git a/Engine/Plugins/Runtime/MicrosoftSpatialAudio/MicroSoftSpatialSound.uplugin b/Engine/Plugins/Runtime/MicrosoftSpatialAudio/MicroSoftSpatialSound.uplugin
--- a/Engine/Plugins/Runtime/MicrosoftSpatialAudio/MicroSoftSpatialSound.uplugin
+++ b/Engine/Plugins/Runtime/MicrosoftSpatialAudio/MicroSoftSpatialSound.uplugin
@@ -14,7 +14,6 @@
 	"IsBetaVersion": false,
 	"Installed": false,
 	"SupportedTargetPlatforms": [
-		"Win64",
 		"XboxOne"
 	],
 	"Modules": [
@@ -24,7 +23,6 @@
 			"LoadingPhase": "PreDefault",
 			"WhitelistPlatforms" :
 			[
-				"Win64",
 				"XboxOne"
 			]
 		}
