Marked the MS spatial audio plugin to be for HoloLens

diff --git a/Engine/Plugins/Runtime/MicrosoftSpatialAudio/MicroSoftSpatialSound.uplugin b/Engine/Plugins/Runtime/MicrosoftSpatialAudio/MicroSoftSpatialSound.uplugin
--- a/Engine/Plugins/Runtime/MicrosoftSpatialAudio/MicroSoftSpatialSound.uplugin
+++ b/Engine/Plugins/Runtime/MicrosoftSpatialAudio/MicroSoftSpatialSound.uplugin
@@ -14,16 +14,15 @@
 	"IsBetaVersion": false,
 	"Installed": false,
 	"SupportedTargetPlatforms": [
-		"XboxOne"
+		"HoloLens"
 	],
 	"Modules": [
 		{
 			"Name": "MicrosoftSpatialSound",
 			"Type": "Runtime",
 			"LoadingPhase": "PreDefault",
-			"WhitelistPlatforms" :
-			[
-				"XboxOne"
+			"WhitelistPlatforms": [
+				"HoloLens"
 			]
 		}
 	]
