Apply distance to spatial sound coordinates for ms spatial sound

diff --git a/Engine/Plugins/Runtime/MicrosoftSpatialAudio/Source/MicrosoftSpatialSound/Private/MicrosoftSpatialSoundPlugin.cpp b/Engine/Plugins/Runtime/MicrosoftSpatialAudio/Source/MicrosoftSpatialSound/Private/MicrosoftSpatialSoundPlugin.cpp
--- a/Engine/Plugins/Runtime/MicrosoftSpatialAudio/Source/MicrosoftSpatialSound/Private/MicrosoftSpatialSoundPlugin.cpp
+++ b/Engine/Plugins/Runtime/MicrosoftSpatialAudio/Source/MicrosoftSpatialSound/Private/MicrosoftSpatialSoundPlugin.cpp
@@ -154,9 +154,9 @@ static void LogMicrosoftSpatialAudioError(HRESULT Result, int32 LineNumber)
 	}
 
 // Function which maps unreal coordinates to MS Spatial sound coordinates
-static FORCEINLINE FVector UnrealToMicrosoftSpatialSoundCoordinates(const FVector& Input)
+static FORCEINLINE FVector UnrealToMicrosoftSpatialSoundCoordinates(const FVector& Input, float InDistance)
 {
-	return { UnrealUnitsToMeters * Input.Y, UnrealUnitsToMeters * Input.X, -UnrealUnitsToMeters * Input.Z };
+	return { UnrealUnitsToMeters * Input.Y * InDistance, UnrealUnitsToMeters * Input.X * InDistance, -UnrealUnitsToMeters * Input.Z * InDistance };
 }
 
 FMicrosoftSpatialSound::FMicrosoftSpatialSound()
@@ -279,7 +279,7 @@ void FMicrosoftSpatialSound::ProcessAudio(const FAudioPluginSourceInputData& Inp
 		UE_LOG(LogMicrosoftSpatialSound, Warning, TEXT("Source circular buffers should be bigger!"));
 	}
 
-	FVector NewPosition = UnrealToMicrosoftSpatialSoundCoordinates(InputData.SpatializationParams->EmitterPosition);
+	FVector NewPosition = UnrealToMicrosoftSpatialSoundCoordinates(InputData.SpatializationParams->EmitterPosition, InputData.SpatializationParams->Distance);
 
 	if (ObjectData.bBuffering && ObjectData.AudioBuffer.Num() > MinFramesRequiredPerObjectUpdate * 5)
 	{
