Fix for microsoft spatial sound not showing up in windows 10 and windows 11.

- API updated and changed it's format to display windows version

diff --git a/Engine/Plugins/Runtime/MicrosoftSpatialAudio/Source/MicrosoftSpatialSound/Private/MicrosoftSpatialSoundPlugin.cpp b/Engine/Plugins/Runtime/MicrosoftSpatialAudio/Source/MicrosoftSpatialSound/Private/MicrosoftSpatialSoundPlugin.cpp
--- a/Engine/Plugins/Runtime/MicrosoftSpatialAudio/Source/MicrosoftSpatialSound/Private/MicrosoftSpatialSoundPlugin.cpp
+++ b/Engine/Plugins/Runtime/MicrosoftSpatialAudio/Source/MicrosoftSpatialSound/Private/MicrosoftSpatialSoundPlugin.cpp
@@ -403,21 +403,10 @@ void FMicrosoftSpatialSoundModule::StartupModule()
 	FString OSSubVersionLabel;
 	FPlatformMisc::GetOSVersions(OSVersionLabel, OSSubVersionLabel);
 
-	// GetOSVersion returns the Win10 release version in the OSVersion rather than the OSSubVersion, so parse it out ourselves
-	OSSubVersionLabel = OSVersionLabel;
-	OSSubVersionLabel.RemoveFromStart("Windows 10 (Release ");
-	OSSubVersionLabel.RemoveFromEnd(")");
-	int32 CurrentVersionNumber = FCString::Atoi(*OSSubVersionLabel);
-
-	if (CurrentVersionNumber < MIN_WIN_10_VERSION_FOR_WMR_SPATSOUND)
+	if (!OSVersionLabel.Contains(TEXT("Windows 10")) && !OSVersionLabel.Contains(TEXT("Windows 11")))
 	{
-		UE_LOG(LogMicrosoftSpatialSound,
-			Warning,
-			TEXT("Microsoft Spatial Sound for Unreal currently only supports windows version '%d' or higher (Current version: '%d')"),
-			MIN_WIN_10_VERSION_FOR_WMR_SPATSOUND,
-			CurrentVersionNumber);
-
-			return;
+		UE_LOG(LogMicrosoftSpatialSound, Warning, TEXT("Microsoft Spatial Sound for Unreal currently only supports windows 10 and windows 11)"));
+		return;
 	}
 
 	// Load the mixed reality interop library
