Fix TARGET_ARCH for arm64ec

#rnx

diff --git a/Engine/Plugins/Runtime/MicrosoftSpatialAudio/Source/MicrosoftSpatialSound/Private/MicrosoftSpatialSoundPlugin.cpp b/Engine/Plugins/Runtime/MicrosoftSpatialAudio/Source/MicrosoftSpatialSound/Private/MicrosoftSpatialSoundPlugin.cpp
--- a/Engine/Plugins/Runtime/MicrosoftSpatialAudio/Source/MicrosoftSpatialSound/Private/MicrosoftSpatialSoundPlugin.cpp
+++ b/Engine/Plugins/Runtime/MicrosoftSpatialAudio/Source/MicrosoftSpatialSound/Private/MicrosoftSpatialSoundPlugin.cpp
@@ -402,6 +402,8 @@ uint32 FMicrosoftSpatialSound::Run()
 #if PLATFORM_CPU_ARM_FAMILY
 	#if (defined(__aarch64__) || defined(_M_ARM64))
 		#define TARGET_ARCH TEXT("arm64")
+	#elif defined(_M_ARM64EC)
+		#define TARGET_ARCH TEXT("x64")
 	#endif
 #else
 	#if PLATFORM_64BITS
